Refactor of Studio Telemetry to move it from a plugin to an Engine level library in preparation for a 5.6 Beta Release
This allows more systems, comamndlets and programs to use FStudioTelemetry API directly and fully replaces the old FStudioAnalytics API that was deprecated as part of 5.4
By default StudioTelemetry relies on the EditorTelemetry and RuntimeTelemetry plugins to Start and End Sessions but can be started or stopped by the applicaiton directly

Specifics:
Moved StudioTelemetry from a plugin to a Developer/StudioTelemetry Engine library
Moved AnalyticsLog and AnalyticsHorde from StudioTelemetry to Runtime/Analytics
Renamed AnalyticsProviderMulticast to AnalyticsProviderBroadcast to avoid confusion with existing AnalyticsMulicast plugin

diff --git a/Engine/Plugins/Editor/AssetSearch/AssetSearch.uplugin b/Engine/Plugins/Editor/AssetSearch/AssetSearch.uplugin
--- a/Engine/Plugins/Editor/AssetSearch/AssetSearch.uplugin
+++ b/Engine/Plugins/Editor/AssetSearch/AssetSearch.uplugin
@@ -25,10 +25,6 @@
 		{
 			"Name": "SQLiteCore",
 			"Enabled": true
-		},
-		{
-			"Name": "StudioTelemetry",
-			"Enabled": true
 		}
 	]
 }
\ No newline at end of file
diff --git a/Engine/Plugins/Experimental/EditorPerformance/EditorPerformance.uplugin b/Engine/Plugins/Experimental/EditorPerformance/EditorPerformance.uplugin
--- a/Engine/Plugins/Experimental/EditorPerformance/EditorPerformance.uplugin
+++ b/Engine/Plugins/Experimental/EditorPerformance/EditorPerformance.uplugin
@@ -20,11 +20,5 @@
 			"Type": "Editor",
 			"LoadingPhase": "PostConfigInit"
 		}
-	],
-	"Plugins": [
-		{
-			"Name": "StudioTelemetry",
-			"Enabled": true
-		}
 	]
 }
\ No newline at end of file
diff --git a/Engine/Plugins/Experimental/EditorTelemetry/EditorTelemetry.uplugin b/Engine/Plugins/Experimental/EditorTelemetry/EditorTelemetry.uplugin
--- a/Engine/Plugins/Experimental/EditorTelemetry/EditorTelemetry.uplugin
+++ b/Engine/Plugins/Experimental/EditorTelemetry/EditorTelemetry.uplugin
@@ -20,11 +20,5 @@
 			"Type": "Editor",
 			"LoadingPhase": "PostConfigInit"
 		}
-	],
-	"Plugins": [
-		{
-			"Name": "StudioTelemetry",
-			"Enabled": true
-		}
 	]
 }
\ No newline at end of file
diff --git a/Engine/Plugins/Experimental/EditorTelemetry/Source/Private/EditorTelemetryModule.cpp b/Engine/Plugins/Experimental/EditorTelemetry/Source/Private/EditorTelemetryModule.cpp
--- a/Engine/Plugins/Experimental/EditorTelemetry/Source/Private/EditorTelemetryModule.cpp
+++ b/Engine/Plugins/Experimental/EditorTelemetry/Source/Private/EditorTelemetryModule.cpp
@@ -10,29 +10,15 @@ IMPLEMENT_MODULE(FEditorTelemetryModule, EditorTelemetry);
 void FEditorTelemetryModule::StartupModule()
 {
 #if !UE_BUILD_SHIPPING
-	if (FStudioTelemetry::Get().IsSessionRunning())
-	{
-		// StudioTelemetry session has already started
-		FEditorTelemetry::Get().StartSession();
-	}
-		
-	FStudioTelemetry::Get().GetOnStartSession().AddLambda([this]()
-		{
-			// StudioTelemetry session has just started
-			FEditorTelemetry::Get().StartSession();
-		}
-	);
-	
-	FStudioTelemetry::Get().GetOnEndSession().AddLambda([this]()
-		{
-			// StudioTelemetry session is about to end
-			FEditorTelemetry::Get().EndSession();
-		}
-	);
+	FStudioTelemetry::Get().StartSession();
+	FEditorTelemetry::Get().StartSession();
 #endif
 }
 
 void FEditorTelemetryModule::ShutdownModule()
 {
+#if !UE_BUILD_SHIPPING
 	FEditorTelemetry::Get().EndSession();
+	FStudioTelemetry::Get().EndSession();
+#endif
 }
diff --git a/Engine/Plugins/Experimental/RuntimeTelemetry/RuntimeTelemetry.uplugin b/Engine/Plugins/Experimental/RuntimeTelemetry/RuntimeTelemetry.uplugin
--- a/Engine/Plugins/Experimental/RuntimeTelemetry/RuntimeTelemetry.uplugin
+++ b/Engine/Plugins/Experimental/RuntimeTelemetry/RuntimeTelemetry.uplugin
@@ -20,11 +20,5 @@
 			"Type": "RuntimeAndProgram",
 			"LoadingPhase": "PostConfigInit"
 		}
-	],
-	"Plugins": [
-		{
-			"Name": "StudioTelemetry",
-			"Enabled": true
-		}
 	]
 }
\ No newline at end of file
diff --git a/Engine/Plugins/Experimental/RuntimeTelemetry/Source/Private/RuntimeTelemetryModule.cpp b/Engine/Plugins/Experimental/RuntimeTelemetry/Source/Private/RuntimeTelemetryModule.cpp
--- a/Engine/Plugins/Experimental/RuntimeTelemetry/Source/Private/RuntimeTelemetryModule.cpp
+++ b/Engine/Plugins/Experimental/RuntimeTelemetry/Source/Private/RuntimeTelemetryModule.cpp
@@ -10,29 +10,15 @@ IMPLEMENT_MODULE(FRuntimeTelemetryModule, RuntimeTelemetry);
 void FRuntimeTelemetryModule::StartupModule()
 {
 #if !UE_BUILD_SHIPPING
-	if (FStudioTelemetry::Get().IsSessionRunning())
-	{
-		// StudioTelemetry session has already started
-		FRuntimeTelemetry::Get().StartSession();
-	}
-
-	FStudioTelemetry::Get().GetOnStartSession().AddLambda([this]()
-		{
-			// StudioTelemetry session has just started
-			FRuntimeTelemetry::Get().StartSession();
-		}
-	);
-	
-	FStudioTelemetry::Get().GetOnEndSession().AddLambda([this]()
-		{
-			// StudioTelemetry session is about to end
-			FRuntimeTelemetry::Get().EndSession();
-		}
-	);
+	FStudioTelemetry::Get().StartSession();
+	FRuntimeTelemetry::Get().StartSession();	
 #endif
 }
 
 void FRuntimeTelemetryModule::ShutdownModule()
 {
-	
+#if !UE_BUILD_SHIPPING
+	FRuntimeTelemetry::Get().EndSession();
+	FStudioTelemetry::Get().EndSession();
+#endif
 }
diff --git a/Engine/Plugins/Experimental/RuntimeTelemetry/Source/Public/RuntimeTelemetry.h b/Engine/Plugins/Experimental/RuntimeTelemetry/Source/Public/RuntimeTelemetry.h
--- a/Engine/Plugins/Experimental/RuntimeTelemetry/Source/Public/RuntimeTelemetry.h
+++ b/Engine/Plugins/Experimental/RuntimeTelemetry/Source/Public/RuntimeTelemetry.h
@@ -24,6 +24,9 @@ public:
 private:	
 
 	void RecordEvent_IoStoreOnDemand(const FString& Context);
+	void RecordEvent_LLMSystemMemory(const FString& Context);
+	void RecordEvent_LLMAssetClassMemory(const FString& Context);
+	void RecordEvent_LLMAssetMemory(const FString& Context);
 };
 
 
diff --git a/Engine/Plugins/Experimental/StudioTelemetry/Source/AnalyticsHorde/AnalyticsHorde.Build.cs b/Engine/Plugins/Experimental/StudioTelemetry/Source/AnalyticsHorde/AnalyticsHorde.Build.cs
deleted file mode 100644
diff --git a/Engine/Plugins/Experimental/StudioTelemetry/Source/AnalyticsHorde/Private/AnalyticsHorde.cpp b/Engine/Plugins/Experimental/StudioTelemetry/Source/AnalyticsHorde/Private/AnalyticsHorde.cpp
deleted file mode 100644
diff --git a/Engine/Plugins/Experimental/StudioTelemetry/Source/AnalyticsHorde/Public/AnalyticsHorde.h b/Engine/Plugins/Experimental/StudioTelemetry/Source/AnalyticsHorde/Public/AnalyticsHorde.h
deleted file mode 100644
diff --git a/Engine/Plugins/Experimental/StudioTelemetry/Source/AnalyticsLog/AnalyticsLog.Build.cs b/Engine/Plugins/Experimental/StudioTelemetry/Source/AnalyticsLog/AnalyticsLog.Build.cs
deleted file mode 100644
diff --git a/Engine/Plugins/Experimental/StudioTelemetry/Source/AnalyticsLog/Private/AnalyticsLog.cpp b/Engine/Plugins/Experimental/StudioTelemetry/Source/AnalyticsLog/Private/AnalyticsLog.cpp
deleted file mode 100644
diff --git a/Engine/Plugins/Experimental/StudioTelemetry/Source/AnalyticsLog/Private/AnalyticsProviderLog.cpp b/Engine/Plugins/Experimental/StudioTelemetry/Source/AnalyticsLog/Private/AnalyticsProviderLog.cpp
deleted file mode 100644
diff --git a/Engine/Plugins/Experimental/StudioTelemetry/Source/AnalyticsLog/Private/AnalyticsProviderLog.h b/Engine/Plugins/Experimental/StudioTelemetry/Source/AnalyticsLog/Private/AnalyticsProviderLog.h
deleted file mode 100644
diff --git a/Engine/Plugins/Experimental/StudioTelemetry/Source/AnalyticsLog/Public/AnalyticsLog.h b/Engine/Plugins/Experimental/StudioTelemetry/Source/AnalyticsLog/Public/AnalyticsLog.h
deleted file mode 100644
diff --git a/Engine/Plugins/Experimental/StudioTelemetry/Source/StudioTelemetry/Private/AnalyticsProviderMulticast.cpp b/Engine/Plugins/Experimental/StudioTelemetry/Source/StudioTelemetry/Private/AnalyticsProviderMulticast.cpp
deleted file mode 100644
diff --git a/Engine/Plugins/Experimental/StudioTelemetry/Source/StudioTelemetry/Private/AnalyticsProviderMulticast.h b/Engine/Plugins/Experimental/StudioTelemetry/Source/StudioTelemetry/Private/AnalyticsProviderMulticast.h
deleted file mode 100644
diff --git a/Engine/Plugins/Experimental/StudioTelemetry/Source/StudioTelemetry/Private/StudioTelemetry.cpp b/Engine/Plugins/Experimental/StudioTelemetry/Source/StudioTelemetry/Private/StudioTelemetry.cpp
deleted file mode 100644
diff --git a/Engine/Plugins/Experimental/StudioTelemetry/Source/StudioTelemetry/Private/StudioTelemetryLog.h b/Engine/Plugins/Experimental/StudioTelemetry/Source/StudioTelemetry/Private/StudioTelemetryLog.h
deleted file mode 100644
diff --git a/Engine/Plugins/Experimental/StudioTelemetry/Source/StudioTelemetry/Public/StudioTelemetry.h b/Engine/Plugins/Experimental/StudioTelemetry/Source/StudioTelemetry/Public/StudioTelemetry.h
deleted file mode 100644
diff --git a/Engine/Plugins/Experimental/StudioTelemetry/Source/StudioTelemetry/StudioTelemetry.build.cs b/Engine/Plugins/Experimental/StudioTelemetry/Source/StudioTelemetry/StudioTelemetry.build.cs
deleted file mode 100644
diff --git a/Engine/Plugins/Experimental/StudioTelemetry/StudioTelemetry.uplugin b/Engine/Plugins/Experimental/StudioTelemetry/StudioTelemetry.uplugin
deleted file mode 100644
diff --git a/Engine/Source/Developer/StudioTelemetry/Private/AnalyticsProviderBroadcast.cpp b/Engine/Source/Developer/StudioTelemetry/Private/AnalyticsProviderBroadcast.cpp
new file mode 100644
--- /dev/null
+++ b/Engine/Source/Developer/StudioTelemetry/Private/AnalyticsProviderBroadcast.cpp
@@ -0,0 +1,265 @@
+// Copyright Epic Games, Inc. All Rights Reserved.
+
+#include "AnalyticsProviderBroadcast.h"
+#include "Analytics.h"
+#include "Interfaces/IAnalyticsProviderModule.h"
+#include "AnalyticsProviderConfigurationDelegate.h"
+#include "Misc/ConfigCacheIni.h"
+#include "HttpModule.h"
+#include "HttpManager.h"
+
+static FString ProviderSection;
+
+FString GetAnalyticsProviderConfiguration(const FString& Name, bool)
+{
+	FString Result;
+	GConfig->GetString(*ProviderSection, *Name, Result, GEngineIni);
+	return Result;
+}
+
+TSharedPtr<FAnalyticsProviderBroadcast> FAnalyticsProviderBroadcast::CreateAnalyticsProvider()
+{
+	return MakeShared<FAnalyticsProviderBroadcast>();
+}
+
+TWeakPtr<IAnalyticsProvider> FAnalyticsProviderBroadcast::GetAnalyticsProvider(const FString& Name)
+{
+	TSharedPtr<IAnalyticsProvider>* ProviderPtr = Providers.Find(Name);
+	return ProviderPtr != nullptr ? *ProviderPtr : TSharedPtr<IAnalyticsProvider>();
+}
+
+FAnalyticsProviderBroadcast::FAnalyticsProviderBroadcast()
+{
+	const FString TelemetryProviderSection(TEXT("StudioTelemetry.Provider"));
+
+	TArray<FString> SectionNames;
+	
+	if (GConfig->GetSectionNames(GEngineIni, SectionNames))
+	{
+		for (const FString& SectionName : SectionNames)
+		{
+			if (SectionName.Find(TelemetryProviderSection) != INDEX_NONE)
+			{
+				ProviderSection = SectionName;
+
+				FString UsageType;
+
+				// Validate the usage type is for this build type
+				if (GConfig->GetString(*ProviderSection, TEXT("UsageType"), UsageType, GEngineIni))
+				{
+					bool IsValidUseCase = false;
+
+#if WITH_EDITOR
+					// Must specify a Editor usage type for this type build
+					if (UsageType.Find(TEXT("Editor")) != INDEX_NONE)
+					{
+						IsValidUseCase |= true;
+					}
+#else 
+					// Must specify a Runtime usage type for all non Editor builds
+					if (UsageType.Find(TEXT("Runtime")) != INDEX_NONE)
+					{
+						IsValidUseCase |= true;
+					}
+#endif
+
+#if WITH_SERVER_CODE
+					// Must specify a Server usage type for this type build
+					if (UsageType.Find(TEXT("Server")) != INDEX_NONE)
+					{
+						IsValidUseCase |= true;
+					}
+#endif
+
+#if WITH_CLIENT_CODE
+					// Must specify a Client usage type for this type build
+					if (UsageType.Find(TEXT("Client")) != INDEX_NONE)
+					{
+						IsValidUseCase |= true;
+					}
+#endif
+
+					if (IsValidUseCase == false)
+					{
+						// This provider is not valid for this use case
+						continue;
+					}
+				}
+				else
+				{
+					// Must always specify a usage type
+					UE_LOG(LogAnalytics, Warning, TEXT("There must be a valid UsageType specified for analytics provider %s"), *ProviderSection);
+					continue;
+				}
+
+				FString ProviderModuleName;
+
+				if (GConfig->GetString(*ProviderSection, TEXT("ProviderModule"), ProviderModuleName, GEngineIni))
+				{
+					TSharedPtr<IAnalyticsProvider> AnalyticsProvider;
+
+					FString Name = GetAnalyticsProviderConfiguration("Name", true);
+
+					if ( Name.IsEmpty() )
+					{ 
+						UE_LOG(LogAnalytics, Error, TEXT("There must be a valid Name specified for analytics provider %s."), *ProviderSection);
+						continue;
+					}
+					else if (Providers.Find(Name) )
+					{
+						UE_LOG(LogAnalytics, Warning, TEXT("An analytics provider with name %s already exists."), *Name);
+						continue;
+					}
+
+					// Try to create the analytics provider
+					AnalyticsProvider = FAnalytics::Get().CreateAnalyticsProvider(FName(ProviderModuleName), FAnalyticsProviderConfigurationDelegate::CreateStatic(&GetAnalyticsProviderConfiguration));
+	
+					if (AnalyticsProvider.IsValid())
+					{
+						UE_LOG(LogAnalytics, Display, TEXT("Created an analytics provider %s from module %s configuration %s [%s]"), *Name, *ProviderModuleName, *GEngineIni, *ProviderSection);
+						Providers.Add(Name, AnalyticsProvider);
+					}
+					else
+					{
+						UE_LOG(LogAnalytics, Warning, TEXT("Unable to create an analytics provider %s from module %s configuration %s [%s]"), *Name, *ProviderModuleName, *GEngineIni, *ProviderSection);
+					}
+				}
+				else
+				{
+					UE_LOG(LogAnalytics, Warning, TEXT("There must be a valid ProviderModule specified for analytics provider %s"), *ProviderSection);
+				}
+			}
+		}
+	}
+}
+
+bool FAnalyticsProviderBroadcast::SetSessionID(const FString& InSessionID)
+{
+	SessionID = InSessionID;
+
+	bool bResult = true;
+
+	for (TProviders::TConstIterator it(Providers); it; ++it)
+	{
+		bResult &= (*it).Value->SetSessionID(InSessionID);
+	}
+
+	return bResult;
+}
+
+FString FAnalyticsProviderBroadcast::GetSessionID() const
+{
+	return SessionID;
+}
+
+void FAnalyticsProviderBroadcast::SetUserID(const FString& InUserID)
+{
+	UserID = InUserID;
+
+	for (TProviders::TConstIterator it(Providers); it; ++it)
+	{
+		(*it).Value->SetUserID(InUserID);
+	}
+}
+
+FString FAnalyticsProviderBroadcast::GetUserID() const
+{
+	return UserID;
+}
+
+void FAnalyticsProviderBroadcast::SetDefaultEventAttributes(TArray<FAnalyticsEventAttribute>&& Attributes)
+{
+	DefaultEventAttributes = Attributes;
+
+	for (TProviders::TConstIterator it(Providers); it; ++it)
+	{
+		(*it).Value->SetDefaultEventAttributes(CopyTemp(DefaultEventAttributes));
+	}
+}
+
+TArray<FAnalyticsEventAttribute> FAnalyticsProviderBroadcast::GetDefaultEventAttributesSafe() const
+{
+	return DefaultEventAttributes;
+}
+
+int32 FAnalyticsProviderBroadcast::GetDefaultEventAttributeCount() const
+{
+	return DefaultEventAttributes.Num();
+}
+
+FAnalyticsEventAttribute FAnalyticsProviderBroadcast::GetDefaultEventAttribute(int AttributeIndex) const
+{
+	return DefaultEventAttributes[AttributeIndex];
+}
+
+bool FAnalyticsProviderBroadcast::StartSession(const TArray<FAnalyticsEventAttribute>& Attributes)
+{
+	bool bResult = true;
+
+	for (TProviders::TConstIterator it(Providers);it;++it)
+	{
+		bResult &= (*it).Value->StartSession(Attributes);
+	}
+	return bResult;
+}
+
+void FAnalyticsProviderBroadcast::EndSession()
+{
+	for (TProviders::TConstIterator it(Providers);it;++it)
+	{
+		TSharedPtr<IAnalyticsProvider> Provider = (*it).Value;
+
+		Provider->EndSession();
+		Provider.Reset();
+	}
+
+	Providers.Reset();	
+}
+
+void FAnalyticsProviderBroadcast::FlushEvents()
+{
+	TRACE_CPUPROFILER_EVENT_SCOPE(FAnalyticsProviderBroadcast::FlushEvents);
+
+	for (TProviders::TConstIterator it(Providers); it; ++it)
+	{
+		(*it).Value->FlushEvents();
+	}
+
+	// It is quite likely that one of the analytics providers is sending data via the FHttpManager so we
+	// need to flush that as well to make sure that the message is sent. 
+	if (FHttpModule* HttpModule = FModuleManager::GetModulePtr<FHttpModule>("HTTP"))
+	{
+		FHttpManager& HttpManager = HttpModule->GetHttpManager();
+		HttpManager.Flush(EHttpFlushReason::FullFlush);
+	}
+}
+
+void FAnalyticsProviderBroadcast::SetRecordEventCallback(OnRecordEvent Callback)
+{
+	RecordEventCallback = Callback;
+}
+
+void FAnalyticsProviderBroadcast::RecordEvent(const FString& EventName, const TArray<FAnalyticsEventAttribute>& Attributes)
+{
+#if UE_BUILD_DEBUG || UE_BUILD_DEVELOPMENT
+	// Expose events that have duplicate attribute names. This is is not handled by the analytics backends in any reliable way.
+	for (int32 index0 = 0; index0 < Attributes.Num(); ++index0)
+	{
+		for (int32 index1 = index0 + 1; index1 < Attributes.Num(); ++index1)
+		{
+			checkf(Attributes[index0].GetName() != Attributes[index1].GetName(), TEXT("Duplicate Attributes Found For Event %s %s==%s"), *EventName, *Attributes[index0].GetName(), *Attributes[index1].GetName());
+		}
+	}
+#endif
+
+	for (TProviders::TConstIterator it(Providers); it; ++it)
+	{
+		(*it).Value->RecordEvent(EventName, Attributes);
+	}
+
+	if (RecordEventCallback)
+	{
+		// Notify any callbacks
+		RecordEventCallback(EventName, Attributes);
+	}
+}
diff --git a/Engine/Source/Developer/StudioTelemetry/Private/AnalyticsProviderBroadcast.h b/Engine/Source/Developer/StudioTelemetry/Private/AnalyticsProviderBroadcast.h
new file mode 100644
--- /dev/null
+++ b/Engine/Source/Developer/StudioTelemetry/Private/AnalyticsProviderBroadcast.h
@@ -0,0 +1,64 @@
+// Copyright Epic Games, Inc. All Rights Reserved.
+
+#pragma once
+
+#include "AnalyticsProviderConfigurationDelegate.h"
+#include "Interfaces/IAnalyticsProvider.h"
+
+/**
+ * Implementation of the IAnalyticsProvider interface that forwards the API calls to an array of IAnalyticsProvider interfaces
+ * 
+ * IAnalyticsProviders are specified in the Engine.ini interface and are constructed and configured automatically.
+ *
+ * AnalyticsLog for example in Engine/Config/BaseEngine.ini
+ * 
+ * ; Studio Telemetry Settings
+ * [StudioTelemetry.Log]
+ * Name=LogAnalytics
+ * ProviderModule=AnalyticsLog
+ * UsageType=Editor
+ * 
+ * This instructs the Broadcast provider to create a Provider from the AnalyticsLog module via the IANalyticsProviderModule interface. 
+ * See FAnalayticsLog and FAnalyticsProviderLog for more details.
+ */
+class FAnalyticsProviderBroadcast : public IAnalyticsProvider
+{
+public:
+
+	using TProviders = TMap<FString, TSharedPtr<IAnalyticsProvider>>;
+
+	typedef TFunction<void(const FString& EventName, const TArray<FAnalyticsEventAttribute>& Attrs)> OnRecordEvent;
+
+
+	FAnalyticsProviderBroadcast();
+
+	static TSharedPtr<FAnalyticsProviderBroadcast> CreateAnalyticsProvider();
+
+	virtual bool StartSession(const TArray<FAnalyticsEventAttribute>& Attributes = {}) override;
+	virtual void EndSession() override;
+	virtual void FlushEvents() override;
+	
+	virtual void SetDefaultEventAttributes(TArray<FAnalyticsEventAttribute>&& Attributes)  override;
+	virtual TArray<FAnalyticsEventAttribute> GetDefaultEventAttributesSafe() const override;
+	virtual int32 GetDefaultEventAttributeCount() const  override;
+	virtual FAnalyticsEventAttribute GetDefaultEventAttribute(int AttributeIndex) const  override;
+
+	virtual bool SetSessionID(const FString& InSessionID) override;
+	virtual void SetUserID(const FString& InUserID) override;
+
+	virtual FString GetSessionID() const override;
+	virtual FString GetUserID() const override;
+	
+	virtual void RecordEvent(const FString& EventName, const TArray<FAnalyticsEventAttribute>& Attributes) override;	
+	void SetRecordEventCallback(OnRecordEvent Callback);
+	bool HasValidProviders() const { return Providers.Num() > 0; }
+	TWeakPtr<IAnalyticsProvider> GetAnalyticsProvider(const FString& Name);
+
+private:
+
+	TProviders								Providers;
+	FString									UserID;
+	FString									SessionID;
+	TArray<FAnalyticsEventAttribute>		DefaultEventAttributes;
+	OnRecordEvent							RecordEventCallback;
+};
diff --git a/Engine/Source/Developer/StudioTelemetry/Private/StudioTelemetry.cpp b/Engine/Source/Developer/StudioTelemetry/Private/StudioTelemetry.cpp
new file mode 100644
--- /dev/null
+++ b/Engine/Source/Developer/StudioTelemetry/Private/StudioTelemetry.cpp
@@ -0,0 +1,337 @@
+// Copyright Epic Games, Inc. All Rights Reserved.
+
+#include "StudioTelemetry.h"
+#include "StudioTelemetryLog.h"
+
+#if WITH_EDITOR
+#include "Horde.h"
+#endif
+
+#include "Analytics.h"
+#include "AnalyticsProviderBroadcast.h"
+#include "AnalyticsTracer.h"
+#include "BuildSettings.h"
+#include "RHI.h"
+
+#include "HAL/PlatformTime.h"
+#include "HAL/PlatformProcess.h"
+#include "HAL/Thread.h"
+
+#include "Misc/CommandLine.h"
+#include "Misc/App.h"
+#include "Misc/EngineVersion.h"
+#include "Misc/EngineBuildSettings.h"
+
+#include "GenericPlatform/GenericPlatformMisc.h"
+
+DEFINE_LOG_CATEGORY(LogStudioTelemetry);
+
+IMPLEMENT_MODULE(FStudioTelemetry, StudioTelemetry)
+
+FStudioTelemetry& FStudioTelemetry::Get()
+{
+	static FStudioTelemetry StudioTelemetryInstance;
+	return StudioTelemetryInstance;
+}
+
+void FStudioTelemetry::SetRecordEventCallback(OnRecordEventCallback Callback )
+{
+	RecordEventCallback = Callback;
+
+	// If the provider already exists then set the callback
+	if (AnalyticsProvider.IsValid())
+	{	
+		AnalyticsProvider->SetRecordEventCallback(RecordEventCallback);
+	}
+}
+
+
+void FStudioTelemetry::StartSession()
+{
+	if (IsSessionRunning())
+	{
+		return;
+	}
+
+	// Load the configuration
+	LoadConfiguration();
+
+	if (Config.bSendTelemetry == false)
+	{
+		// We did not wish to send any telemetry events
+		return;
+	}
+
+	FScopeLock ScopeLock(&CriticalSection);
+
+	AnalyticsProvider = FAnalyticsProviderBroadcast::CreateAnalyticsProvider();
+
+	if (AnalyticsProvider.IsValid())
+	{
+		TArray<FAnalyticsEventAttribute> DefaultEventAttributes;
+
+		const FString UserID = FPlatformProcess::UserName(false);
+		const FString ProjectName = FApp::GetProjectName();
+		FString ComputerName = FPlatformProcess::ComputerName();
+
+		FString ProjectIDString;
+		GConfig->GetString(TEXT("/Script/EngineSettings.GeneralProjectSettings"), TEXT("ProjectID"), ProjectIDString, GGameIni);
+
+		FGuid ProjectID;
+
+		if (!ProjectIDString.IsEmpty())
+		{
+			TArray<FString> Elements;
+			if (ProjectIDString.ParseIntoArray(Elements, TEXT("=")) == 5)
+			{
+				ProjectID = FGuid(FCString::Atoi(*(Elements[1])), FCString::Atoi(*(Elements[2])), FCString::Atoi(*(Elements[3])), FCString::Atoi(*(Elements[4])));
+			}
+			else
+			{
+				ProjectID = FGuid(ProjectIDString);
+			}
+		}
+
+		FGuid SessionID = FApp::GetInstanceId();
+
+		FString SessionLabel;
+		FParse::Value(FCommandLine::Get(), TEXT("SessionLabel="), SessionLabel);
+
+		// Set the default event attributes, these will always be sent to telemetry for every event
+		DefaultEventAttributes.Emplace(TEXT("ProjectName"), ProjectName);
+		DefaultEventAttributes.Emplace(TEXT("ProjectID"), ProjectID);
+
+		DefaultEventAttributes.Emplace(TEXT("Session_ID"), SessionID.ToString(EGuidFormats::DigitsWithHyphensInBraces));
+		DefaultEventAttributes.Emplace(TEXT("Session_Label"), SessionLabel);
+		DefaultEventAttributes.Emplace(TEXT("Session_StartUTC"), FDateTime::UtcNow().ToUnixTimestampDecimal());
+
+		DefaultEventAttributes.Emplace(TEXT("Build_Configuration"), LexToString(FApp::GetBuildConfiguration()));
+		DefaultEventAttributes.Emplace(TEXT("Build_IsInternalBuild"), FEngineBuildSettings::IsInternalBuild());
+		DefaultEventAttributes.Emplace(TEXT("Build_IsPerforceBuild"), FEngineBuildSettings::IsPerforceBuild());
+		DefaultEventAttributes.Emplace(TEXT("Build_IsPromotedBuild"), FApp::GetEngineIsPromotedBuild() == 0 ? false : true);
+		DefaultEventAttributes.Emplace(TEXT("Build_BranchName"), FApp::GetBranchName().ToLower());
+		DefaultEventAttributes.Emplace(TEXT("Build_Changelist"), BuildSettings::GetCurrentChangelist());
+
+		DefaultEventAttributes.Emplace(TEXT("Config_IsEditor"), GIsEditor);
+		DefaultEventAttributes.Emplace(TEXT("Config_IsUnattended"), FApp::IsUnattended());
+		DefaultEventAttributes.Emplace(TEXT("Config_IsBuildMachine"), GIsBuildMachine);
+		DefaultEventAttributes.Emplace(TEXT("Config_IsRunningCommandlet"), IsRunningCommandlet());
+		DefaultEventAttributes.Emplace(TEXT("Config_IsDebuggerPresent"), FPlatformMisc::IsDebuggerPresent());
+
+		// Only send user data if requested
+		if (Config.bSendUserData == true)
+		{
+			DefaultEventAttributes.Emplace(TEXT("User_ID"), UserID);
+			DefaultEventAttributes.Emplace(TEXT("Application_Commandline"), FCommandLine::Get());
+		}
+
+		// Only send hardware data if requested
+		if (Config.bSendHardwareData == true)
+		{
+			DefaultEventAttributes.Emplace(TEXT("Hardware_Platform"), FString(FPlatformProperties::IniPlatformName()));
+			DefaultEventAttributes.Emplace(TEXT("Hardware_GPU"), GRHIAdapterName);
+			DefaultEventAttributes.Emplace(TEXT("Hardware_CPU"), FPlatformMisc::GetCPUBrand());
+			DefaultEventAttributes.Emplace(TEXT("Hardware_CPU_Cores_Physical"), FPlatformMisc::NumberOfCores());
+			DefaultEventAttributes.Emplace(TEXT("Hardware_CPU_Cores_Logical"), FPlatformMisc::NumberOfCoresIncludingHyperthreads());
+			DefaultEventAttributes.Emplace(TEXT("Hardware_RAM"), static_cast<uint64>(FPlatformMemory::GetStats().TotalPhysical));
+			DefaultEventAttributes.Emplace(TEXT("Hardware_ComputerName"), ComputerName);
+		}
+
+		// Only send OS data if requested
+		if (Config.bSendOSData == true)
+		{
+			FString OSVersionLabel;
+			FString OSSubVersionLabel;
+
+			FPlatformMisc::GetOSVersions(OSVersionLabel, OSSubVersionLabel);
+
+			DefaultEventAttributes.Emplace(TEXT("OS_Version"), FPlatformMisc::GetOSVersion());
+			DefaultEventAttributes.Emplace(TEXT("OS_VersionLabel"), OSVersionLabel);
+			DefaultEventAttributes.Emplace(TEXT("OS_VersionSubLabel"), OSSubVersionLabel);
+			DefaultEventAttributes.Emplace(TEXT("OS_ID"), FPlatformMisc::GetOperatingSystemId());
+		}
+
+#if WITH_EDITOR
+		if (!FHorde::GetJobId().IsEmpty())
+		{
+			// Only send Horde data if applicable
+			DefaultEventAttributes.Emplace(TEXT("Horde_ServerURL"), FHorde::GetServerURL());
+			DefaultEventAttributes.Emplace(TEXT("Horde_TemplateID"), FHorde::GetTemplateId());
+			DefaultEventAttributes.Emplace(TEXT("Horde_TemplateName"), FHorde::GetTemplateName());
+			DefaultEventAttributes.Emplace(TEXT("Horde_JobURL"), FHorde::GetJobURL());
+			DefaultEventAttributes.Emplace(TEXT("Horde_JobID"), FHorde::GetJobId());
+			DefaultEventAttributes.Emplace(TEXT("Horde_StepName"), FHorde::GetStepName());
+			DefaultEventAttributes.Emplace(TEXT("Horde_StepID"), FHorde::GetStepId());
+			DefaultEventAttributes.Emplace(TEXT("Horde_StepURL"), FHorde::GetStepURL());
+			DefaultEventAttributes.Emplace(TEXT("Horde_BatchID"), FHorde::GetBatchId());
+		}
+#endif
+
+		// Set up the analytics provider
+		AnalyticsProvider->SetUserID(UserID);
+		AnalyticsProvider->SetSessionID(SessionID.ToString(EGuidFormats::DigitsWithHyphensInBraces));
+		AnalyticsProvider->SetDefaultEventAttributes(MoveTemp(DefaultEventAttributes));
+		AnalyticsProvider->SetRecordEventCallback(RecordEventCallback);
+
+		// Start the analytics session
+		AnalyticsProvider->StartSession();
+
+		// Create the IAnalyticsTracer interface
+		AnalyticsTracer = FAnalytics::Get().CreateAnalyticsTracer();
+		AnalyticsTracer->SetProvider(AnalyticsProvider);
+		AnalyticsTracer->StartSession();
+
+		// Bind the pre-exit callback
+		FCoreDelegates::OnEnginePreExit.AddRaw(&FStudioTelemetry::Get(), &FStudioTelemetry::EndSession);
+
+		OnStartSession.Broadcast();
+
+		UE_LOG(LogStudioTelemetry, Log, TEXT("Started StudioTelemetry Session"));
+	}
+}
+
+void FStudioTelemetry::EndSession()
+{
+	if (IsSessionRunning())
+	{
+		FScopeLock ScopeLock(&CriticalSection);
+
+		OnEndSession.Broadcast();
+
+		// End session for the tracer and the provider
+		if (AnalyticsTracer.IsValid())
+		{
+			AnalyticsTracer->EndSession();
+			AnalyticsTracer.Reset();
+		}
+
+		if (AnalyticsProvider.IsValid())
+		{
+			AnalyticsProvider->EndSession();
+			AnalyticsProvider->FlushEvents();
+			AnalyticsProvider.Reset();
+		}
+
+		UE_LOG(LogStudioTelemetry, Log, TEXT("Ended StudioTelemetry Session"));
+	}
+}
+
+bool FStudioTelemetry::IsSessionRunning() const
+{
+	return AnalyticsProvider.IsValid();
+}
+
+void FStudioTelemetry::LoadConfiguration()
+{
+	const FString TelemetryConfigurationSection(TEXT("StudioTelemetry.Config"));
+
+	// Look for the configuration settings in the Engine.ini files
+	TArray<FString> SectionNames;
+
+	if (GConfig->GetSectionNames(GEngineIni, SectionNames))
+	{
+		for (const FString& SectionName : SectionNames)
+		{
+			if (SectionName.Find(TelemetryConfigurationSection) != INDEX_NONE)
+			{
+				GConfig->GetBool(*SectionName, TEXT("SendTelemetry="), Config.bSendTelemetry, GEngineIni);
+				GConfig->GetBool(*SectionName, TEXT("SendUserData="), Config.bSendUserData, GEngineIni);
+				GConfig->GetBool(*SectionName, TEXT("SendHardwareData="), Config.bSendHardwareData, GEngineIni);
+				GConfig->GetBool(*SectionName, TEXT("SendOSData="), Config.bSendOSData, GEngineIni);
+			}
+		}
+	}
+
+	// Parse the commandline for any local configuration overrides
+	FParse::Bool(FCommandLine::Get(), TEXT("ST_SendTelemetry="), Config.bSendTelemetry);
+	FParse::Bool(FCommandLine::Get(), TEXT("ST_SendUserData="), Config.bSendUserData);
+	FParse::Bool(FCommandLine::Get(), TEXT("ST_SendHardwareData="), Config.bSendHardwareData);
+	FParse::Bool(FCommandLine::Get(), TEXT("ST_SendOSData="), Config.bSendOSData);
+}
+
+void FStudioTelemetry::RecordEvent(const FString& EventName, const TArray<FAnalyticsEventAttribute>& Attributes)
+{
+	FScopeLock ScopeLock(&CriticalSection);
+	if (AnalyticsProvider.IsValid())
+	{
+		AnalyticsProvider->RecordEvent(EventName, Attributes);
+		OnRecordEvent.Broadcast(EventName, Attributes);
+	}
+}
+
+void FStudioTelemetry::RecordEvent(const FName CategoryName, const FString& EventName, const TArray<FAnalyticsEventAttribute>& Attributes)
+{
+	FScopeLock ScopeLock(&CriticalSection);
+	if (AnalyticsProvider.IsValid())
+	{
+		AnalyticsProvider->RecordEvent(EventName, Attributes);
+		OnRecordEvent.Broadcast(EventName, Attributes);
+	}
+}
+
+void FStudioTelemetry::FlushEvents()
+{
+	FScopeLock ScopeLock(&CriticalSection);
+	if (AnalyticsProvider.IsValid())
+	{
+		AnalyticsProvider->FlushEvents();
+	}
+}
+
+void FStudioTelemetry::RecordEventToProvider(const FString& ProviderName, const FString& EventName, const TArray<FAnalyticsEventAttribute>& Attributes)
+{
+	FScopeLock ScopeLock(&CriticalSection);
+	TSharedPtr<IAnalyticsProvider> NamedProvider = GetProvider(ProviderName).Pin();
+
+	if (NamedProvider.IsValid())
+	{
+		NamedProvider->RecordEvent(EventName, Attributes);
+	}
+}
+
+TWeakPtr<IAnalyticsProvider> FStudioTelemetry::GetProvider()
+{
+	return AnalyticsProvider;
+}
+
+TWeakPtr<IAnalyticsProvider> FStudioTelemetry::GetProvider(const FString& Name)
+{
+	return AnalyticsProvider.IsValid()? AnalyticsProvider->GetAnalyticsProvider(Name) : TWeakPtr<IAnalyticsProvider>();
+}
+
+TWeakPtr<IAnalyticsTracer> FStudioTelemetry::GetTracer()
+{
+	return AnalyticsTracer;
+}
+
+TSharedPtr<IAnalyticsSpan> FStudioTelemetry::StartSpan(const FName Name, const TArray<FAnalyticsEventAttribute>& AdditionalAttributes)
+{
+	return AnalyticsTracer.IsValid() ? AnalyticsTracer->StartSpan(Name, AdditionalAttributes) : TSharedPtr<IAnalyticsSpan>();
+}
+
+TSharedPtr<IAnalyticsSpan> FStudioTelemetry::StartSpan(const FName Name, TSharedPtr<IAnalyticsSpan> ParentSpan, const TArray<FAnalyticsEventAttribute>& AdditionalAttributes)
+{
+	return AnalyticsTracer.IsValid() ? AnalyticsTracer->StartSpan(Name, ParentSpan, AdditionalAttributes)  : TSharedPtr<IAnalyticsSpan>();
+}
+
+bool FStudioTelemetry::EndSpan(TSharedPtr<IAnalyticsSpan> Span, const TArray<FAnalyticsEventAttribute>& AdditionalAttributes)
+{
+	return AnalyticsTracer.IsValid() ? AnalyticsTracer->EndSpan(Span, AdditionalAttributes) : false;
+}
+
+bool FStudioTelemetry::EndSpan(const FName Name, const TArray<FAnalyticsEventAttribute>& AdditionalAttributes)
+{
+	return AnalyticsTracer.IsValid() ? AnalyticsTracer->EndSpan(Name, AdditionalAttributes) : false;
+}
+
+TSharedPtr<IAnalyticsSpan> FStudioTelemetry::GetSpan(const FName Name)
+{
+	return AnalyticsTracer.IsValid() ? AnalyticsTracer->GetSpan(Name) : TSharedPtr<IAnalyticsSpan>();
+}
+
+TSharedPtr<IAnalyticsSpan> FStudioTelemetry::GetSessionSpan() const
+{
+	return AnalyticsTracer.IsValid() ? AnalyticsTracer->GetSessionSpan() : TSharedPtr<IAnalyticsSpan>();
+}
+
+
diff --git a/Engine/Source/Developer/StudioTelemetry/Private/StudioTelemetryLog.h b/Engine/Source/Developer/StudioTelemetry/Private/StudioTelemetryLog.h
new file mode 100644
--- /dev/null
+++ b/Engine/Source/Developer/StudioTelemetry/Private/StudioTelemetryLog.h
@@ -0,0 +1,10 @@
+// Copyright Epic Games, Inc. All Rights Reserved.
+
+#pragma once
+
+#include "CoreMinimal.h"
+
+/**
+ * StudioTelemetry plugin logging
+ */
+DECLARE_LOG_CATEGORY_EXTERN(LogStudioTelemetry, Log, All);
diff --git a/Engine/Source/Developer/StudioTelemetry/Public/StudioTelemetry.h b/Engine/Source/Developer/StudioTelemetry/Public/StudioTelemetry.h
new file mode 100644
--- /dev/null
+++ b/Engine/Source/Developer/StudioTelemetry/Public/StudioTelemetry.h
@@ -0,0 +1,148 @@
+// Copyright Epic Games, Inc. All Rights Reserved.
+
+#pragma once
+
+#include "Interfaces/IAnalyticsProvider.h"
+#include "Interfaces/IAnalyticsTracer.h"
+#include "Containers/Array.h"
+#include "Containers/UnrealString.h"
+#include "CoreMinimal.h"
+#include "Modules/ModuleManager.h"
+#include "Templates/SharedPointer.h"
+
+class FAnalyticsProviderBroadcast;
+/**
+ * Studio Telemetry API
+ * 
+ * Notes:
+ * Interface for adding Studio level Telemetry to products. Studio Telemetry will never function in shipping builds.
+ * Developers are encouraged to add post their own development telemetry events via this API.
+ * Developers can implement their own IAnalyticsProviderModule where custom recording of Studio Telemetry events to their own Analytics Backend is desired.
+ * Custom AnalyticsProviders can be added to the plugin via the .ini. See FAnalyticsProviderLog or FAnalyticsProviderET for example.
+ * Telemetry events are recored to all registered IAnalyticsProviders supplied in the .ini file using the FAnalyticsProviderBroadcast provider, except where specifically recorded with the RecordEvent(ProviderName,.. ) API below
+ */
+class FStudioTelemetry : public IModuleInterface
+{
+public:
+
+	typedef TFunction<void(const FString& EventName, const TArray<FAnalyticsEventAttribute>& Attrs)> OnRecordEventCallback;
+
+	/** Check whether the module is available*/
+	static STUDIOTELEMETRY_API bool IsAvailable() { return FStudioTelemetry::Get().IsSessionRunning(); }
+
+	/** Access to the module singleton*/
+	static STUDIOTELEMETRY_API FStudioTelemetry& Get();
+
+	/** Access to the a specific named analytics provider within the system*/
+	STUDIOTELEMETRY_API TWeakPtr<IAnalyticsProvider> GetProvider(const FString& ProviderName);
+
+	/** Access to the broadcast analytics provider for the system*/
+	STUDIOTELEMETRY_API TWeakPtr<IAnalyticsProvider> GetProvider();
+
+	/** Access to the tracer for the system*/
+	STUDIOTELEMETRY_API TWeakPtr<IAnalyticsTracer> GetTracer();
+
+	/** Starts a new analytics session*/
+	STUDIOTELEMETRY_API void StartSession();
+
+	/** Ends an existing analytics session*/
+	STUDIOTELEMETRY_API void EndSession();
+
+	/** Is Session Running */
+	STUDIOTELEMETRY_API bool IsSessionRunning() const;
+	
+	/** Thread safe method to record an event to all registered analytics providers*/
+	STUDIOTELEMETRY_API void RecordEvent(const FString& EventName, const TArray<FAnalyticsEventAttribute>& Attributes = {});
+
+	/** Thread safe method to record an event to all registered analytics providers*/
+	STUDIOTELEMETRY_API void RecordEvent(const FName CategoryName, const FString& EventName, const TArray<FAnalyticsEventAttribute>& Attributes = {});
+
+	/** Thread safe method to record an event to the specifically named analytics provider */
+	STUDIOTELEMETRY_API void RecordEventToProvider(const FString& ProviderName, const FString& EventName, const TArray<FAnalyticsEventAttribute>& Attributes = {});
+
+	/** Thread safe method to flush all events on all registered analytics providers*/
+	STUDIOTELEMETRY_API void FlushEvents();
+
+	/** Start a new span specifying the parent*/
+	STUDIOTELEMETRY_API TSharedPtr<IAnalyticsSpan> StartSpan(const FName Name, const TArray<FAnalyticsEventAttribute>& AdditionalAttributes = {});
+
+	/** Start a new span specifying the parent*/
+	STUDIOTELEMETRY_API TSharedPtr<IAnalyticsSpan> StartSpan(const FName Name, TSharedPtr<IAnalyticsSpan> ParentSpan, const TArray<FAnalyticsEventAttribute>& AdditionalAttributes = {});
+
+	/** End an existing span*/
+	STUDIOTELEMETRY_API bool EndSpan(TSharedPtr<IAnalyticsSpan> Span, const TArray<FAnalyticsEventAttribute>& AdditionalAttributes = {});
+
+	/** End an existing span by name*/
+	STUDIOTELEMETRY_API bool EndSpan(const FName Name, const TArray<FAnalyticsEventAttribute>& AdditionalAttributes = {});
+
+	/** Get an active span by name, non active spans will not be available*/
+	STUDIOTELEMETRY_API TSharedPtr<IAnalyticsSpan> GetSpan(const FName Name);
+
+	/** Get the root session span*/
+	STUDIOTELEMETRY_API TSharedPtr<IAnalyticsSpan> GetSessionSpan() const;
+
+	/** Callback for interception of telemetry events recording that can be used by Developers to send telemetry events to their own back end, though it is recommended that Developers implement their own IAnalyticsProvider via their own IAnalyticsProviderModule*/
+	STUDIOTELEMETRY_API void SetRecordEventCallback(OnRecordEventCallback);
+
+	/** Delegates for event callbacks **/
+	DECLARE_MULTICAST_DELEGATE(FOnStartSession);
+	STUDIOTELEMETRY_API FOnStartSession& GetOnStartSession() { return OnStartSession; }
+
+	DECLARE_MULTICAST_DELEGATE(FOnEndSession);
+	STUDIOTELEMETRY_API FOnEndSession& GetOnEndSession() { return OnEndSession; }
+
+	DECLARE_MULTICAST_DELEGATE_TwoParams(FOnRecordEvent, const FString&, const TArray<FAnalyticsEventAttribute>&);
+	STUDIOTELEMETRY_API FOnRecordEvent& GetOnRecordEvent() { return OnRecordEvent; }
+
+	/** Scoped Span helper class **/
+	class ScopedSpan
+	{
+	public:
+		ScopedSpan(const FName Name, const TArray<FAnalyticsEventAttribute>& AdditionalAttributes = {} )
+		{
+			if (FStudioTelemetry::Get().IsAvailable())
+			{
+				Span = FStudioTelemetry::Get().StartSpan(Name, AdditionalAttributes);
+			}
+		}
+
+		~ScopedSpan()
+		{
+			if (FStudioTelemetry::Get().IsAvailable())
+			{
+				FStudioTelemetry::Get().EndSpan(Span);
+			}
+		}
+	private:
+
+		TSharedPtr<IAnalyticsSpan> Span;
+	};
+
+private:
+
+	/** Configure the plugin*/
+	void LoadConfiguration();
+
+	struct FConfig
+	{
+		bool bSendTelemetry = true; // Only send telemetry data if we have been requested to
+		bool bSendUserData = false;  // Never send user data unless specifically asked to
+		bool bSendHardwareData = true; // Always send hardware data unless specifically asked not to
+		bool bSendOSData = true; // Always send operating system data unless specifically asked not to
+	};
+	
+	FCriticalSection						CriticalSection;
+	TSharedPtr<FAnalyticsProviderBroadcast>	AnalyticsProvider;
+	TSharedPtr<IAnalyticsTracer>			AnalyticsTracer;
+	OnRecordEventCallback					RecordEventCallback;
+	FGuid									SessionGUID;
+	FConfig									Config;
+	FOnStartSession							OnStartSession;
+	FOnEndSession							OnEndSession;
+	FOnRecordEvent							OnRecordEvent;
+};
+
+// Useful macros for scoped spans
+#define STUDIO_TELEMETRY_SPAN_SCOPE(Name) FStudioTelemetry::ScopedSpan PREPROCESSOR_JOIN(ScopedSpan, __LINE__)(TEXT(#Name));
+#define STUDIO_TELEMETRY_START_SPAN(Name) if (FStudioTelemetry::Get().IsAvailable()) { FStudioTelemetry::Get().StartSpan(TEXT(#Name));}
+#define STUDIO_TELEMETRY_END_SPAN(Name) if (FStudioTelemetry::Get().IsAvailable()) { FStudioTelemetry::Get().EndSpan(TEXT(#Name));}
diff --git a/Engine/Source/Developer/StudioTelemetry/StudioTelemetry.build.cs b/Engine/Source/Developer/StudioTelemetry/StudioTelemetry.build.cs
new file mode 100644
--- /dev/null
+++ b/Engine/Source/Developer/StudioTelemetry/StudioTelemetry.build.cs
@@ -0,0 +1,42 @@
+// Copyright Epic Games, Inc. All Rights Reserved.
+
+using System.IO;
+using UnrealBuildTool;
+
+public class StudioTelemetry : ModuleRules
+{
+	public StudioTelemetry(ReadOnlyTargetRules Target) : base(Target)
+	{
+		PCHUsage = ModuleRules.PCHUsageMode.UseExplicitOrSharedPCHs;
+		
+		PublicDependencyModuleNames.AddRange(
+			new string[] {
+				"Core",
+				"HTTP"
+			}
+		);
+
+		PrivateDependencyModuleNames.AddRange(
+			new string[] {
+				"CoreUObject",
+                "EngineSettings",
+                "BuildSettings",
+                "Analytics",
+                "AnalyticsET",
+				"TelemetryUtils",
+				"RHI",
+			}
+		);
+
+		if (Target.bBuildEditor)
+		{
+			PrivateDependencyModuleNames.AddRange(
+				new string[] {
+					"AnalyticsLog",
+					"AnalyticsHorde",
+					"Horde"
+				}
+			);
+		}
+	}
+}
diff --git a/Engine/Source/Runtime/Analytics/AnalyticsHorde/AnalyticsHorde.Build.cs b/Engine/Source/Runtime/Analytics/AnalyticsHorde/AnalyticsHorde.Build.cs
new file mode 100644
--- /dev/null
+++ b/Engine/Source/Runtime/Analytics/AnalyticsHorde/AnalyticsHorde.Build.cs
@@ -0,0 +1,28 @@
+// Copyright Epic Games, Inc. All Rights Reserved.
+
+namespace UnrealBuildTool.Rules
+{
+	public class AnalyticsHorde : ModuleRules
+	{
+		public AnalyticsHorde(ReadOnlyTargetRules Target) : base(Target)
+		{
+			PrivateDependencyModuleNames.AddRange(
+				new string[] {
+				"Core",
+				"CoreUObject",
+				"Analytics",
+				"AnalyticsET"
+				}
+			);
+
+			if (Target.bBuildEditor)
+			{
+					PrivateDependencyModuleNames.AddRange(
+						new string[] {
+						"Horde"
+						}
+					);
+				}
+			}
+	}
+}
diff --git a/Engine/Source/Runtime/Analytics/AnalyticsHorde/Private/AnalyticsHorde.cpp b/Engine/Source/Runtime/Analytics/AnalyticsHorde/Private/AnalyticsHorde.cpp
new file mode 100644
--- /dev/null
+++ b/Engine/Source/Runtime/Analytics/AnalyticsHorde/Private/AnalyticsHorde.cpp
@@ -0,0 +1,42 @@
+// Copyright Epic Games, Inc. All Rights Reserved.
+
+#include "AnalyticsHorde.h"
+#include "AnalyticsET.h"
+#include "IAnalyticsProviderET.h"
+
+#if WITH_EDITOR
+#include "Horde.h"
+#endif
+
+IMPLEMENT_MODULE( FAnalyticsHorde, AnalyticsHorde );
+
+void FAnalyticsHorde::StartupModule()
+{
+}
+
+void FAnalyticsHorde::ShutdownModule()
+{
+}
+
+TSharedPtr<IAnalyticsProvider> FAnalyticsHorde::CreateAnalyticsProvider(const FAnalyticsProviderConfigurationDelegate& GetConfigValue) const
+{
+	if (GetConfigValue.IsBound())
+	{
+		TSharedPtr<IAnalyticsProviderET> AnalyicsProviderET = FAnalyticsET::Get().CreateAnalyticsProviderET(GetConfigValue);
+
+#if WITH_EDITOR
+		// Check if we have specifcied a Horde URL in the environment. This allows jobs running on Horde to send telemetry directly to the server they were run on
+		const FString HordeServerURL = FHorde::GetServerURL();
+
+		if (AnalyicsProviderET.IsValid() && !HordeServerURL.IsEmpty())
+		{
+			TArray<FString> AltDomains;
+			AnalyicsProviderET->SetURLEndpoint(HordeServerURL, AltDomains);
+		}
+#endif
+		
+		return AnalyicsProviderET;
+	}
+
+	return nullptr;
+}
diff --git a/Engine/Source/Runtime/Analytics/AnalyticsHorde/Public/AnalyticsHorde.h b/Engine/Source/Runtime/Analytics/AnalyticsHorde/Public/AnalyticsHorde.h
new file mode 100644
--- /dev/null
+++ b/Engine/Source/Runtime/Analytics/AnalyticsHorde/Public/AnalyticsHorde.h
@@ -0,0 +1,50 @@
+// Copyright Epic Games, Inc. All Rights Reserved.
+
+#pragma once
+
+#include "AnalyticsProviderConfigurationDelegate.h"
+#include "Interfaces/IAnalyticsProviderModule.h"
+#include "Modules/ModuleManager.h"
+
+class IAnalyticsProvider;
+
+/**
+ *  Public implementation of IAnalyticsProviderModule
+ */
+class FAnalyticsHorde : public IAnalyticsProviderModule
+{
+	//--------------------------------------------------------------------------
+	// Module functionality
+	//--------------------------------------------------------------------------
+public:
+	/**
+	 * Singleton-like access to this module's interface.  This is just for convenience!
+	 * Beware of calling this during the shutdown phase, though.  Your module might have been unloaded already.
+	 *
+	 * @return Returns singleton instance, loading the module on demand if needed
+	 */
+	static inline FAnalyticsHorde& Get()
+	{
+		return FModuleManager::LoadModuleChecked< FAnalyticsHorde >( "AnalyticsHorde" );
+	}
+
+	//--------------------------------------------------------------------------
+	// Configuration functionality
+	//--------------------------------------------------------------------------
+public:
+	
+	//--------------------------------------------------------------------------
+	// provider factory functions
+	//--------------------------------------------------------------------------
+public:
+	/**
+	 * IAnalyticsProviderModule interface.
+	 * Creates the analytics provider given a configuration delegate.
+	 * The keys required exactly match the field names in the Config object. 
+	 */
+	ANALYTICSHORDE_API virtual TSharedPtr<IAnalyticsProvider> CreateAnalyticsProvider(const FAnalyticsProviderConfigurationDelegate& GetConfigValue) const override;
+
+private:
+	ANALYTICSHORDE_API virtual void StartupModule() override;
+	ANALYTICSHORDE_API virtual void ShutdownModule() override;
+};
diff --git a/Engine/Source/Runtime/Analytics/AnalyticsLog/AnalyticsLog.Build.cs b/Engine/Source/Runtime/Analytics/AnalyticsLog/AnalyticsLog.Build.cs
new file mode 100644
--- /dev/null
+++ b/Engine/Source/Runtime/Analytics/AnalyticsLog/AnalyticsLog.Build.cs
@@ -0,0 +1,23 @@
+// Copyright Epic Games, Inc. All Rights Reserved.
+
+namespace UnrealBuildTool.Rules
+{
+	public class AnalyticsLog : ModuleRules
+	{
+		public AnalyticsLog(ReadOnlyTargetRules Target) : base(Target)
+		{
+			PublicDependencyModuleNames.AddRange(
+				new string[] {
+						"Core"
+				}
+			);
+
+			PrivateDependencyModuleNames.AddRange(
+				new string[] {
+				"CoreUObject",
+				"Analytics"
+				}
+			);
+		}
+	}
+}
diff --git a/Engine/Source/Runtime/Analytics/AnalyticsLog/Private/AnalyticsLog.cpp b/Engine/Source/Runtime/Analytics/AnalyticsLog/Private/AnalyticsLog.cpp
new file mode 100644
--- /dev/null
+++ b/Engine/Source/Runtime/Analytics/AnalyticsLog/Private/AnalyticsLog.cpp
@@ -0,0 +1,25 @@
+// Copyright Epic Games, Inc. All Rights Reserved.
+
+#include "AnalyticsLog.h"
+#include "AnalyticsProviderLog.h"
+
+IMPLEMENT_MODULE( FAnalyticsLog, AnalyticsLog );
+
+void FAnalyticsLog::StartupModule()
+{
+}
+
+void FAnalyticsLog::ShutdownModule()
+{
+}
+
+TSharedPtr<IAnalyticsProvider> FAnalyticsLog::CreateAnalyticsProvider(const FAnalyticsProviderConfigurationDelegate& GetConfigValue) const
+{
+	if (GetConfigValue.IsBound())
+	{
+		return MakeShared<FAnalyticsProviderLog>(GetConfigValue);
+	}
+
+	return nullptr;
+}
+
diff --git a/Engine/Source/Runtime/Analytics/AnalyticsLog/Private/AnalyticsProviderLog.cpp b/Engine/Source/Runtime/Analytics/AnalyticsLog/Private/AnalyticsProviderLog.cpp
new file mode 100644
--- /dev/null
+++ b/Engine/Source/Runtime/Analytics/AnalyticsLog/Private/AnalyticsProviderLog.cpp
@@ -0,0 +1,169 @@
+// Copyright Epic Games, Inc. All Rights Reserved.
+
+#include "AnalyticsProviderLog.h"
+
+#include "Analytics.h"
+#include "HAL/FileManager.h"
+#include "Misc/Paths.h"
+#include "Tasks/Pipe.h"
+
+FAnalyticsProviderLog::FAnalyticsProviderLog(const FAnalyticsProviderConfigurationDelegate& GetConfigValue)
+{
+	FString FileName = GetConfigValue.Execute(TEXT("FileName"), true);
+
+	if (FileName.IsEmpty())
+	{
+		// Use default filename
+		FileName = TEXT("Telemetry.json");
+	}
+
+	FString FolderPath = GetConfigValue.Execute(TEXT("FolderPath"), true);
+
+	if (FolderPath.IsEmpty())
+	{
+		// Use default output path
+		FolderPath = FPaths::ProjectSavedDir() / TEXT("Telemetry");
+	}
+
+	// Create the full output path
+	FString FilePath = FolderPath / FileName;
+	FileWriter = TUniquePtr<FArchive>(IFileManager::Get().CreateFileWriter(*FilePath, FILEWRITE_EvenIfReadOnly));
+
+	if (FileWriter.IsValid())
+	{
+		WriterPipe = MakeUnique<UE::Tasks::FPipe>(TEXT("FAnalyticsProviderLog_Writer"));
+	}
+}
+
+FAnalyticsProviderLog::~FAnalyticsProviderLog()
+{
+
+}
+
+bool FAnalyticsProviderLog::SetSessionID(const FString& InSessionID)
+{
+	SessionID = InSessionID;
+	return true;
+}
+
+FString FAnalyticsProviderLog::GetSessionID() const
+{
+	return SessionID;
+}
+
+void FAnalyticsProviderLog::SetUserID(const FString& InUserID)
+{
+	UserID = InUserID;
+
+}
+
+FString FAnalyticsProviderLog::GetUserID() const
+{
+	return UserID;
+}
+
+void FAnalyticsProviderLog::FlushEvents()
+{
+}
+
+void FAnalyticsProviderLog::SetDefaultEventAttributes(TArray<FAnalyticsEventAttribute>&& Attributes)
+{
+	DefaultEventAttributes = Attributes;
+}
+
+TArray<FAnalyticsEventAttribute> FAnalyticsProviderLog::GetDefaultEventAttributesSafe() const
+{
+	return DefaultEventAttributes;
+}
+
+int32 FAnalyticsProviderLog::GetDefaultEventAttributeCount() const
+{
+	return DefaultEventAttributes.Num();
+}
+
+FAnalyticsEventAttribute FAnalyticsProviderLog::GetDefaultEventAttribute(int AttributeIndex) const
+{
+	return DefaultEventAttributes[AttributeIndex];
+}
+
+bool FAnalyticsProviderLog::StartSession(const TArray<FAnalyticsEventAttribute>& Attributes)
+{
+	TRACE_CPUPROFILER_EVENT_SCOPE(FAnalyticsProviderLog::StartSession);
+
+	RecordEvent(TEXT("StartSession"), Attributes);
+
+	return true;
+}
+
+void FAnalyticsProviderLog::EndSession()
+{
+	TRACE_CPUPROFILER_EVENT_SCOPE(FAnalyticsProviderLog::EndSession);
+
+	RecordEvent(TEXT("EndSession"));
+
+	if (WriterPipe.IsValid())
+	{
+		TRACE_CPUPROFILER_EVENT_SCOPE(FAnalyticsProviderLog::EndSession::ClosePipe);
+		WriterPipe->WaitUntilEmpty();
+		WriterPipe.Reset();
+	}
+
+	if (FileWriter.IsValid())
+	{
+		FileWriter->Flush();
+		FileWriter->Close();
+
+		FileWriter.Reset();
+	}
+}
+
+void FAnalyticsProviderLog::RecordEvent(const FString& EventName, const TArray<FAnalyticsEventAttribute>& Attributes)
+{	
+	TRACE_CPUPROFILER_EVENT_SCOPE(FAnalyticsProviderLog::RecordEvent);
+
+	static uint32 RecordId(0);
+
+	if (FileWriter.IsValid() && WriterPipe.IsValid())
+	{
+		TStringBuilder<1024> Builder;
+
+		// Log event as Newline - delimited JSON
+		Builder.Appendf(TEXT("{\"EventName\":\"%s\""), *EventName);
+
+		// Add the event timestamp field
+		Builder.Appendf(TEXT(",\"TimestampUTC\":%f"), FDateTime::UtcNow().ToUnixTimestampDecimal());
+
+		// Add the record Id and increment it
+		Builder.Appendf(TEXT(",\"RecordId\":%d"), RecordId++);
+
+		// Accumulate all the attributes together. We could have had two loops but this seems cleaner
+		TArray<FAnalyticsEventAttribute> EventAttributes(DefaultEventAttributes);
+		EventAttributes.Append(Attributes);
+
+		// Add all the attributes
+		for (const FAnalyticsEventAttribute& Attribute : EventAttributes)
+		{
+			// This should be almost nearly true, but we should check and JSON'ify as needed
+			if (Attribute.IsJsonFragment())
+			{
+				Builder.Appendf(TEXT(",\"%s\":%s"), *Attribute.GetName(), *Attribute.GetValue());
+			}
+			else
+			{
+				Builder.Appendf(TEXT(",\"%s\":\"%s\""), *Attribute.GetName(), *Attribute.GetValue());
+			}
+		}
+
+		Builder.Append(TEXT("}"));
+
+		FString FinalOutput = Builder.ToString();
+
+		WriterPipe->Launch(TEXT("FAnalyticsProviderLog_WriteJob"), [this, FinalOutput = MoveTemp(FinalOutput)]()
+			{
+				TRACE_CPUPROFILER_EVENT_SCOPE(FAnalyticsProviderLog_WriteJob);
+				
+				FileWriter->Logf(TEXT("%s"), *FinalOutput);
+				FileWriter->Flush();
+			});
+	}
+}
diff --git a/Engine/Source/Runtime/Analytics/AnalyticsLog/Private/AnalyticsProviderLog.h b/Engine/Source/Runtime/Analytics/AnalyticsLog/Private/AnalyticsProviderLog.h
new file mode 100644
--- /dev/null
+++ b/Engine/Source/Runtime/Analytics/AnalyticsLog/Private/AnalyticsProviderLog.h
@@ -0,0 +1,47 @@
+// Copyright Epic Games, Inc. All Rights Reserved.
+
+#pragma once
+
+#include "AnalyticsProviderConfigurationDelegate.h"
+#include "Interfaces/IAnalyticsProvider.h"
+
+namespace UE::Tasks { class FPipe; }
+
+/**
+ * Implementation of the IAnalyticsProviderET interface that exports telemetry events to a file as Newline - delimited JSON
+ * By default,the log file is written to Saved/Telemetry folder of the application.
+ * FileName and FolderPath can be overridden in the configuration
+ * Here as a simple example of how a developer might implement and configure their own analytics provider for use with the StudioTelemerty plugin
+ */
+class FAnalyticsProviderLog : public IAnalyticsProvider
+{
+public:
+
+	FAnalyticsProviderLog(const FAnalyticsProviderConfigurationDelegate& GetConfigValue);
+	~FAnalyticsProviderLog();
+
+	virtual bool StartSession(const TArray<FAnalyticsEventAttribute>& Attributes = {}) override;
+	virtual void EndSession() override;
+
+	virtual void FlushEvents() override;
+
+	virtual void SetDefaultEventAttributes(TArray<FAnalyticsEventAttribute>&& Attributes)  override;
+	virtual TArray<FAnalyticsEventAttribute> GetDefaultEventAttributesSafe() const override;
+	virtual int32 GetDefaultEventAttributeCount() const  override;
+	virtual FAnalyticsEventAttribute GetDefaultEventAttribute(int AttributeIndex) const  override;
+	virtual bool SetSessionID(const FString& InSessionID) override;
+	virtual void SetUserID(const FString& InUserID) override;
+
+	virtual FString GetSessionID() const override;
+	virtual FString GetUserID() const override;
+	
+	virtual void RecordEvent(const FString& EventName, const TArray<FAnalyticsEventAttribute>& Attributes = {}) override;
+
+private:
+
+	FString									UserID;
+	FString									SessionID;
+	TArray<FAnalyticsEventAttribute>		DefaultEventAttributes;	
+	TUniquePtr<FArchive>					FileWriter;
+	TUniquePtr<UE::Tasks::FPipe>			WriterPipe;
+};
diff --git a/Engine/Source/Runtime/Analytics/AnalyticsLog/Public/AnalyticsLog.h b/Engine/Source/Runtime/Analytics/AnalyticsLog/Public/AnalyticsLog.h
new file mode 100644
--- /dev/null
+++ b/Engine/Source/Runtime/Analytics/AnalyticsLog/Public/AnalyticsLog.h
@@ -0,0 +1,46 @@
+// Copyright Epic Games, Inc. All Rights Reserved.
+
+#pragma once
+
+#include "AnalyticsProviderConfigurationDelegate.h"
+#include "Containers/Array.h"
+#include "Containers/UnrealString.h"
+#include "CoreMinimal.h"
+#include "HAL/Platform.h"
+#include "HAL/PlatformCrt.h"
+#include "Interfaces/IAnalyticsProviderModule.h"
+#include "Modules/ModuleManager.h"
+#include "Templates/SharedPointer.h"
+#include "Templates/UnrealTemplate.h"
+
+class IAnalyticsProvider;
+
+/**
+ *  Public implementation of IAnalyticsProviderModule that returns a JSON formatted FAnalayticsProviderLog object
+ */
+class FAnalyticsLog : public IAnalyticsProviderModule
+{
+	//--------------------------------------------------------------------------
+	// Module functionality
+	//--------------------------------------------------------------------------
+public:
+	/**
+	 * Singleton-like access to this module's interface.  This is just for convenience!
+	 * Beware of calling this during the shutdown phase, though.  Your module might have been unloaded already
+	 * @return Returns singleton instance, loading the module on demand if needed
+	 */
+	static inline FAnalyticsLog& Get()
+	{
+		return FModuleManager::LoadModuleChecked< FAnalyticsLog >( "AnalyticsLog" );
+	}
+
+	/**
+	 * IAnalyticsProviderModule interface.
+	 */
+	ANALYTICSLOG_API virtual TSharedPtr<IAnalyticsProvider> CreateAnalyticsProvider(const FAnalyticsProviderConfigurationDelegate& GetConfigValue) const override;
+
+private:
+	ANALYTICSLOG_API virtual void StartupModule() override;
+	ANALYTICSLOG_API virtual void ShutdownModule() override;
+};
+
